Fix incorrect IMPLEMENT_MATERIAL_SHADER_TYPE for Google Tango overlay shaders after Dev-Rendering merge.

diff --git a/Engine/Plugins/Runtime/GoogleTango/Source/TangoRenderer/Private/TangoVideoOverlayRendererRHI.cpp b/Engine/Plugins/Runtime/GoogleTango/Source/TangoRenderer/Private/TangoVideoOverlayRendererRHI.cpp
--- a/Engine/Plugins/Runtime/GoogleTango/Source/TangoRenderer/Private/TangoVideoOverlayRendererRHI.cpp
+++ b/Engine/Plugins/Runtime/GoogleTango/Source/TangoRenderer/Private/TangoVideoOverlayRendererRHI.cpp
@@ -246,7 +246,7 @@ public:
 	}
 };
 
-IMPLEMENT_MATERIAL_SHADER_TYPE(,FTangoCameraOverlayVS,TEXT("PostProcessMaterialShaders"),TEXT("MainVS_ES2"),SF_Vertex);
+IMPLEMENT_MATERIAL_SHADER_TYPE(,FTangoCameraOverlayVS,TEXT("/Engine/Private/PostProcessMaterialShaders.usf"),TEXT("MainVS_ES2"),SF_Vertex);
 
 class FTangoCameraOverlayPS : public FMaterialShader
 {
@@ -288,7 +288,7 @@ public:
 
 };
 
-IMPLEMENT_MATERIAL_SHADER_TYPE(,FTangoCameraOverlayPS,TEXT("PostProcessMaterialShaders"),TEXT("MainPS_ES2"),SF_Pixel);
+IMPLEMENT_MATERIAL_SHADER_TYPE(,FTangoCameraOverlayPS,TEXT("/Engine/Private/PostProcessMaterialShaders.usf"),TEXT("MainPS_ES2"),SF_Pixel);
 
 void FTangoVideoOverlayRendererRHI::RenderVideoOverlay_RenderThread(FRHICommandListImmediate& RHICmdList, FSceneView& InView)
 {
