Fix deprectation warning

diff --git a/Engine/Plugins/Runtime/GoogleTango/Source/TangoRenderer/Private/TangoVideoOverlayRendererRHI.cpp b/Engine/Plugins/Runtime/GoogleTango/Source/TangoRenderer/Private/TangoVideoOverlayRendererRHI.cpp
--- a/Engine/Plugins/Runtime/GoogleTango/Source/TangoRenderer/Private/TangoVideoOverlayRendererRHI.cpp
+++ b/Engine/Plugins/Runtime/GoogleTango/Source/TangoRenderer/Private/TangoVideoOverlayRendererRHI.cpp
@@ -326,7 +326,7 @@ void FTangoVideoOverlayRendererRHI::RenderVideoOverlay_RenderThread(FRHICommandL
 
 		if (OverlayVertexBufferRHI.IsValid() && OverlayIndexBufferRHI.IsValid())
 		{
-			RHICmdList.SetStreamSource(0, OverlayVertexBufferRHI, sizeof(FFilterVertex), 0);
+			RHICmdList.SetStreamSource(0, OverlayVertexBufferRHI, 0);
 			RHICmdList.DrawIndexedPrimitive(
 				OverlayIndexBufferRHI,
 				PT_TriangleList,
